Project Settings - Adding a way to easily introduce per platform settings objects that are correctly saved outside the game ini, and instead store their configuration to each platform's NDA'ed location, if they're a confidential platform.  These platform options are kept in sync with what platforms you have enabled for your project (though a restart is required to see that reflected in the settings).

The base class of any per platform settings.  The pattern for using these is as follows.

Step 1) Subclass UPlatformSettings, UMyPerPlatformSettings : public UPlatformSettings.

Step 2) For your system should already have a UDeveloperSettings that you created so that
        users can customize other properties for your system in the project.  On that class
        you need to create a property of type FPerPlatformSettings, e.g.
        UPROPERTY(EditAnywhere, Category=Platform)
        FPerPlatformSettings PlatformOptions

Step 3) In your UDeveloperSettings subclasses construct, there should be a line like this,
        PlatformOptions.Settings = UPlatformSettings::GetAllPlatformSettings<UMyPerPlatformSettings>();
        This will actually ensure that you initialize the settings exposed in the editor to whatever
        the current platform configuration is for them.

Step 4) Nothing else needed.  In your system code, you will just call
        UMyPerPlatformSettings* MySettings = UPlatformSettings::GetSettingsForPlatform<UMyPerPlatformSettings>()
        that will get you the current settings for the active platform, or the simulated platform in the editor.

diff --git a/Engine/Source/Developer/SettingsEditor/Private/Customizations/PerPlatformSettingsCustomization.cpp b/Engine/Source/Developer/SettingsEditor/Private/Customizations/PerPlatformSettingsCustomization.cpp
new file mode 100644
--- /dev/null
+++ b/Engine/Source/Developer/SettingsEditor/Private/Customizations/PerPlatformSettingsCustomization.cpp
@@ -0,0 +1,72 @@
+// Copyright Epic Games, Inc. All Rights Reserved.
+
+#include "PerPlatformSettingsCustomization.h"
+#include "DetailWidgetRow.h"
+#include "Widgets/SBoxPanel.h"
+#include "Widgets/Images/SImage.h"
+#include "EditorStyleSet.h"
+#include "Engine/PlatformSettings.h"
+#include "IDetailGroup.h"
+#include "IDetailChildrenBuilder.h"
+
+TSharedRef<IPropertyTypeCustomization> FPerPlatformSettingsCustomization::MakeInstance()
+{
+	return MakeShareable(new FPerPlatformSettingsCustomization);
+}
+
+FPerPlatformSettingsCustomization::FPerPlatformSettingsCustomization()
+{
+}
+
+void FPerPlatformSettingsCustomization::CustomizeHeader(TSharedRef<IPropertyHandle> PropertyHandle, FDetailWidgetRow& HeaderRow, IPropertyTypeCustomizationUtils& CustomizationUtils)
+{
+	HeaderRow
+	.WholeRowContent()
+	.MaxDesiredWidth(TOptional<float>())
+	.MaxDesiredWidth(TOptional<float>())
+	[
+		PropertyHandle->CreatePropertyNameWidget()
+	];
+}
+
+void FPerPlatformSettingsCustomization::CustomizeChildren(TSharedRef<IPropertyHandle> PropertyHandle, IDetailChildrenBuilder& ChildBuilder, IPropertyTypeCustomizationUtils& CustomizationUtils)
+{
+	TSharedPtr<IPropertyHandle> SettingsProperty = PropertyHandle->GetChildHandle(GET_MEMBER_NAME_CHECKED(FPerPlatformSettings, Settings));
+	TSharedPtr<IPropertyHandleArray> SettingsPropertyArray = SettingsProperty->AsArray();
+
+	uint32 NumSettings;
+	SettingsPropertyArray->GetNumElements(NumSettings);
+	for (uint32 Index = 0; Index < NumSettings; ++Index)
+	{
+		TSharedPtr<IPropertyHandle> PlatformSettingsProperty = SettingsPropertyArray->GetElement(Index);
+		if (PlatformSettingsProperty->IsValidHandle())
+		{
+			UObject* PlatformSettingsObject;
+			FPropertyAccess::Result Result = PlatformSettingsProperty->GetValue(PlatformSettingsObject);
+			if (Result == FPropertyAccess::Success)
+			{
+				if (UPlatformSettings* PlatformSettings = Cast<UPlatformSettings>(PlatformSettingsObject))
+				{
+					FAddPropertyParams Params = FAddPropertyParams()
+						.UniqueId(PlatformSettings->GetFName())
+						.AllowChildren(true)
+						.CreateCategoryNodes(false);
+
+					IDetailGroup& PlatformGroup = ChildBuilder.AddGroup(PlatformSettings->GetFName(), FText::FromString(PlatformSettings->GetPlatformIniName()));
+
+					TSharedPtr<IPropertyHandle> PlatformSettingsPropertyArrayEntry = PlatformSettingsProperty->GetChildHandle(0);
+
+					//FDetailWidgetRow& PlatformGroup = ChildBuilder.AddCustomRow(FText::FromString(PlatformSettings->GetPlatformIniName()));
+
+					uint32 PlatformPropertiesChildHandleCount;
+					PlatformSettingsPropertyArrayEntry->GetNumChildren(PlatformPropertiesChildHandleCount);
+					for (uint32 PlatformPropertyIndex = 0; PlatformPropertyIndex < PlatformPropertiesChildHandleCount; PlatformPropertyIndex++)
+					{
+						TSharedPtr<IPropertyHandle> PlatformPropertyHandle = PlatformSettingsPropertyArrayEntry->GetChildHandle(PlatformPropertyIndex);
+						PlatformGroup.AddPropertyRow(PlatformPropertyHandle.ToSharedRef());
+					}
+				}
+			}
+		}
+	}
+}
\ No newline at end of file
diff --git a/Engine/Source/Developer/SettingsEditor/Private/Customizations/PerPlatformSettingsCustomization.h b/Engine/Source/Developer/SettingsEditor/Private/Customizations/PerPlatformSettingsCustomization.h
new file mode 100644
--- /dev/null
+++ b/Engine/Source/Developer/SettingsEditor/Private/Customizations/PerPlatformSettingsCustomization.h
@@ -0,0 +1,25 @@
+// Copyright Epic Games, Inc. All Rights Reserved.
+
+#pragma once
+
+#include "CoreMinimal.h"
+#include "Misc/Attribute.h"
+#include "Input/Reply.h"
+#include "IDetailCustomization.h"
+#include "PropertyHandle.h"
+
+//////////////////////////////////////////////////////////////////////////
+// FPerPlatformSettingsCustomization
+
+class FPerPlatformSettingsCustomization : public IPropertyTypeCustomization
+{
+public:
+	// Makes a new instance of this detail layout class for a specific detail view requesting it
+	static TSharedRef<IPropertyTypeCustomization> MakeInstance();
+
+	virtual void CustomizeHeader(TSharedRef<IPropertyHandle> PropertyHandle, class FDetailWidgetRow& HeaderRow, IPropertyTypeCustomizationUtils& CustomizationUtils) override;
+	virtual void CustomizeChildren(TSharedRef<IPropertyHandle> PropertyHandle, class IDetailChildrenBuilder& ChildBuilder, IPropertyTypeCustomizationUtils& CustomizationUtils) override;
+
+private:
+	FPerPlatformSettingsCustomization();
+};
diff --git a/Engine/Source/Developer/SettingsEditor/Private/Customizations/PlatformSettingsCustomization.h b/Engine/Source/Developer/SettingsEditor/Private/Customizations/PlatformSettingsCustomization.h
new file mode 100644
--- /dev/null
+++ b/Engine/Source/Developer/SettingsEditor/Private/Customizations/PlatformSettingsCustomization.h
@@ -0,0 +1,24 @@
+// Copyright Epic Games, Inc. All Rights Reserved.
+
+#pragma once
+
+#include "CoreMinimal.h"
+#include "Misc/Attribute.h"
+#include "Input/Reply.h"
+#include "IDetailCustomization.h"
+#include "PropertyHandle.h"
+
+//////////////////////////////////////////////////////////////////////////
+// FPerPlatformSettingsCustomization
+
+class FPerPlatformSettingsCustomization : public IPropertyTypeCustomization
+{
+public:
+	static TSharedRef<IPropertyTypeCustomization> MakeInstance();
+
+	virtual void CustomizeHeader(TSharedRef<IPropertyHandle> PropertyHandle, class FDetailWidgetRow& HeaderRow, IPropertyTypeCustomizationUtils& CustomizationUtils) override;
+	virtual void CustomizeChildren(TSharedRef<IPropertyHandle> PropertyHandle, class IDetailChildrenBuilder& ChildBuilder, IPropertyTypeCustomizationUtils& CustomizationUtils) override;
+
+private:
+	FPerPlatformSettingsCustomization();
+};
diff --git a/Engine/Source/Developer/SettingsEditor/Private/SettingsEditorModule.cpp b/Engine/Source/Developer/SettingsEditor/Private/SettingsEditorModule.cpp
--- a/Engine/Source/Developer/SettingsEditor/Private/SettingsEditorModule.cpp
+++ b/Engine/Source/Developer/SettingsEditor/Private/SettingsEditorModule.cpp
@@ -16,6 +16,10 @@
 #include "ISettingsEditorModule.h"
 #include "ISettingsModule.h"
 #include "Engine/DeveloperSettings.h"
+#include "Engine/PlatformSettings.h"
+#include "PropertyEditorModule.h"
+#include "PropertyEditorDelegates.h"
+#include "Customizations/PlatformSettingsCustomization.h"
 
 #define LOCTEXT_NAMESPACE "SSettingsEditor"
 
@@ -79,8 +83,6 @@ public:
 		}
 	}
 
-
-
 private:
 	void OnRestartClicked()
 	{
@@ -133,6 +135,12 @@ public:
 
 	virtual void StartupModule() override
 	{
+		FPropertyEditorModule& PropertyModule = FModuleManager::LoadModuleChecked<FPropertyEditorModule>("PropertyEditor");
+		PropertyModule.RegisterCustomPropertyTypeLayout(
+			StaticStruct<FPerPlatformSettings>()->GetFName(),
+			FOnGetPropertyTypeCustomizationInstance::CreateStatic(&FPerPlatformSettingsCustomization::MakeInstance)
+		);
+
 		FModuleManager::Get().OnModulesChanged().AddRaw(this, &FSettingsEditorModule::ModulesChangesCallback);
 	}
 
@@ -238,37 +246,42 @@ private:
 						continue;
 					}
 
-					FRegisteredSettings Registered;
-					Registered.ContainerName = Settings->GetContainerName();
-					Registered.CategoryName = Settings->GetCategoryName();
-					Registered.SectionName = Settings->GetSectionName();
-
-					TSharedPtr<SWidget> CustomWidget = Settings->GetCustomSettingsWidget();
-					if ( CustomWidget.IsValid() )
-					{
-						// Add Settings
-						SettingsModule.RegisterSettings(Registered.ContainerName, Registered.CategoryName, Registered.SectionName,
-							Settings->GetSectionText(),
-							Settings->GetSectionDescription(),
-							CustomWidget.ToSharedRef()
-							);
-					}
-					else
-					{
-						// Add Settings
-						SettingsModule.RegisterSettings(Registered.ContainerName, Registered.CategoryName, Registered.SectionName,
-							Settings->GetSectionText(),
-							Settings->GetSectionDescription(),
-							Settings
-							);
-					}
-
-					AutoDiscoveredSettings.Add(Registered);
+					RegisterDeveloperSettings(SettingsModule, Settings);
 				}
 			}
 		}
 	}
 
+	void RegisterDeveloperSettings(ISettingsModule& SettingsModule, UDeveloperSettings* Settings)
+	{
+		FRegisteredSettings Registered;
+		Registered.ContainerName = Settings->GetContainerName();
+		Registered.CategoryName = Settings->GetCategoryName();
+		Registered.SectionName = Settings->GetSectionName();
+
+		TSharedPtr<SWidget> CustomWidget = Settings->GetCustomSettingsWidget();
+		if (CustomWidget.IsValid())
+		{
+			// Add Settings
+			SettingsModule.RegisterSettings(Registered.ContainerName, Registered.CategoryName, Registered.SectionName,
+				Settings->GetSectionText(),
+				Settings->GetSectionDescription(),
+				CustomWidget.ToSharedRef()
+			);
+		}
+		else
+		{
+			// Add Settings
+			SettingsModule.RegisterSettings(Registered.ContainerName, Registered.CategoryName, Registered.SectionName,
+				Settings->GetSectionText(),
+				Settings->GetSectionDescription(),
+				Settings
+			);
+		}
+
+		AutoDiscoveredSettings.Add(Registered);
+	}
+
 	void UnregisterAutoDiscoveredSettings(ISettingsModule& SettingsModule)
 	{
 		// Unregister any auto discovers settings.
diff --git a/Engine/Source/Developer/SettingsEditor/Private/Widgets/SSettingsEditor.cpp b/Engine/Source/Developer/SettingsEditor/Private/Widgets/SSettingsEditor.cpp
--- a/Engine/Source/Developer/SettingsEditor/Private/Widgets/SSettingsEditor.cpp
+++ b/Engine/Source/Developer/SettingsEditor/Private/Widgets/SSettingsEditor.cpp
@@ -221,6 +221,12 @@ void SSettingsEditor::NotifyPostChange( const FPropertyChangedEvent& PropertyCha
 				{
 					Section->Save();
 				}
+				// Some files being edited might have an array element, but they may also not have a corresponding section for inlined
+				// external objects, for them if they're DefaultConfig, we update them here.
+				else if (ObjectBeingEdited->GetClass()->HasAnyClassFlags(CLASS_DefaultConfig))
+				{
+					ObjectBeingEdited->UpdateDefaultConfigFile();
+				}
 
 				if (bIsNewFile && bIsSourceControlled)
 				{
diff --git a/Engine/Source/Runtime/DeveloperSettings/DeveloperSettings.Build.cs b/Engine/Source/Runtime/DeveloperSettings/DeveloperSettings.Build.cs
--- a/Engine/Source/Runtime/DeveloperSettings/DeveloperSettings.Build.cs
+++ b/Engine/Source/Runtime/DeveloperSettings/DeveloperSettings.Build.cs
@@ -11,6 +11,7 @@ public class DeveloperSettings : ModuleRules
 				"Core",
 				"CoreUObject",
                 "SlateCore",
+				"Projects"
 			}
 		);
 	}
diff --git a/Engine/Source/Runtime/DeveloperSettings/Private/DeveloperSettingsModule.cpp b/Engine/Source/Runtime/DeveloperSettings/Private/DeveloperSettingsModule.cpp
--- a/Engine/Source/Runtime/DeveloperSettings/Private/DeveloperSettingsModule.cpp
+++ b/Engine/Source/Runtime/DeveloperSettings/Private/DeveloperSettingsModule.cpp
@@ -1,6 +1,22 @@
 // Copyright Epic Games, Inc. All Rights Reserved.
 
-#include "DeveloperSettingsModule.h"
 #include "Modules/ModuleManager.h"
+#include "Modules/ModuleInterface.h"
 
-IMPLEMENT_MODULE( IDeveloperSettingsModule, DeveloperSettings );
+class FDeveloperSettingsModule : public IModuleInterface
+{
+public:
+	/** IModuleInterface implementation */
+	virtual void StartupModule() override;
+	virtual void ShutdownModule() override;
+};
+
+void FDeveloperSettingsModule::StartupModule()
+{
+}
+
+void FDeveloperSettingsModule::ShutdownModule()
+{
+}
+
+IMPLEMENT_MODULE(FDeveloperSettingsModule, DeveloperSettings );
diff --git a/Engine/Source/Runtime/DeveloperSettings/Private/Engine/PlatformSettings.cpp b/Engine/Source/Runtime/DeveloperSettings/Private/Engine/PlatformSettings.cpp
new file mode 100644
--- /dev/null
+++ b/Engine/Source/Runtime/DeveloperSettings/Private/Engine/PlatformSettings.cpp
@@ -0,0 +1,91 @@
+// Copyright Epic Games, Inc. All Rights Reserved.
+
+#include "Engine/PlatformSettings.h"
+#include "HAL/IConsoleManager.h"
+#include "UObject/UnrealType.h"
+#include "UObject/EnumProperty.h"
+#include "UObject/PropertyPortFlags.h"
+#include "Misc/DataDrivenPlatformInfoRegistry.h"
+#include "HAL/PlatformProperties.h"
+#include "Interfaces/IProjectManager.h"
+#include "ProjectDescriptor.h"
+
+#if WITH_EDITOR
+FString UPlatformSettings::SimulatedEditorPlatform;
+#endif
+
+UPlatformSettings::UPlatformSettings(const FObjectInitializer& ObjectInitializer)
+	: Super(ObjectInitializer)
+{
+	
+}
+
+#if WITH_EDITOR
+TArray<TObjectPtr<UPlatformSettings>> UPlatformSettings::GetAllPlatformSettings(TSubclassOf<UPlatformSettings> SettingsClass)
+{
+	TArray<TObjectPtr<UPlatformSettings>> Settings;
+
+	const FProjectDescriptor* Project = IProjectManager::Get().GetCurrentProject();
+	
+	const TMap<FName, FDataDrivenPlatformInfo>& AllPlatforms = FDataDrivenPlatformInfoRegistry::GetAllPlatformInfos();
+
+	for (auto& PlatformKVP : AllPlatforms)
+	{
+		const FName PlatformIniName = PlatformKVP.Key;
+		const FDataDrivenPlatformInfo& PlatformInfo = PlatformKVP.Value;
+
+		if (PlatformInfo.bIsFakePlatform)
+		{
+			continue;
+		}
+
+		if (Project->TargetPlatforms.IsEmpty() || Project->TargetPlatforms.Contains(PlatformIniName))
+		{
+			Settings.Add(GetSettingsForPlatformInternal(SettingsClass, PlatformIniName.ToString()));
+		}
+	}
+	
+	return Settings;
+}
+#endif
+
+TObjectPtr<UPlatformSettings> UPlatformSettings::GetSettingsForPlatform(TSubclassOf<UPlatformSettings> SettingsClass)
+{
+	static UPlatformSettings* ThisPlatformsSettings = GetSettingsForPlatformInternal(SettingsClass, FPlatformProperties::IniPlatformName());
+	
+#if WITH_EDITOR
+	if (GIsEditor && !SimulatedEditorPlatform.IsEmpty())
+	{
+		return GetSettingsForPlatformInternal(SettingsClass, SimulatedEditorPlatform);
+	}
+#endif
+	
+	return ThisPlatformsSettings;
+}
+
+UPlatformSettings* UPlatformSettings::GetSettingsForPlatformInternal(TSubclassOf<UPlatformSettings> SettingsClass, FString TargetIniPlatformName)
+{
+	if (!ensure(SettingsClass.Get()))
+	{
+		return nullptr;
+	}
+
+	const FString PlatformSettingsName = SettingsClass->GetName() + TEXT("_") + TargetIniPlatformName;
+	
+	UPlatformSettings* PlatformSettingsForThisClass = FindObject<UPlatformSettings>(GetTransientPackage(), *PlatformSettingsName);
+	if (PlatformSettingsForThisClass == nullptr)
+	{
+		const TMap<FName, FDataDrivenPlatformInfo>& AllPlatforms = FDataDrivenPlatformInfoRegistry::GetAllPlatformInfos();
+		if (AllPlatforms.Contains(*TargetIniPlatformName))
+		{
+			PlatformSettingsForThisClass = NewObject<UPlatformSettings>(GetTransientPackage(), SettingsClass, FName(*PlatformSettingsName));
+			PlatformSettingsForThisClass->ConfigPlatformName = TargetIniPlatformName;
+			PlatformSettingsForThisClass->AddToRoot();
+			PlatformSettingsForThisClass->InitializePlatformDefaults();
+
+			PlatformSettingsForThisClass->LoadConfig();
+		}
+	}
+	
+	return PlatformSettingsForThisClass;
+}
diff --git a/Engine/Source/Runtime/DeveloperSettings/Public/DeveloperSettingsModule.h b/Engine/Source/Runtime/DeveloperSettings/Public/DeveloperSettingsModule.h
deleted file mode 100644
diff --git a/Engine/Source/Runtime/DeveloperSettings/Public/Engine/PlatformSettings.h b/Engine/Source/Runtime/DeveloperSettings/Public/Engine/PlatformSettings.h
new file mode 100644
--- /dev/null
+++ b/Engine/Source/Runtime/DeveloperSettings/Public/Engine/PlatformSettings.h
@@ -0,0 +1,97 @@
+// Copyright Epic Games, Inc. All Rights Reserved.
+
+#pragma once
+
+#include "CoreMinimal.h"
+#include "UObject/ObjectMacros.h"
+#include "UObject/UObjectGlobals.h"
+#include "UObject/Object.h"
+#include "Templates/SubclassOf.h"
+#include "DeveloperSettings.h"
+
+#include "PlatformSettings.generated.h"
+
+USTRUCT()
+struct FPerPlatformSettings
+{
+	GENERATED_BODY()
+
+public:
+	UPROPERTY(Instanced, Transient, EditAnywhere, EditFixedSize, Category = Layout)
+	TArray<TObjectPtr<UPlatformSettings>> Settings;
+};
+
+/**
+ * The base class of any per platform settings.  The pattern for using these is as follows.
+ * 
+ * Step 1) Subclass UPlatformSettings, UMyPerPlatformSettings : public UPlatformSettings.
+ * 
+ * Step 2) For your system should already have a UDeveloperSettings that you created so that
+ *         users can customize other properties for your system in the project.  On that class
+ *         you need to create a property of type FPerPlatformSettings, e.g. 
+ *         UPROPERTY(EditAnywhere, Category=Platform)
+ *         FPerPlatformSettings PlatformOptions
+ * 
+ * Step 3) In your UDeveloperSettings subclasses construct, there should be a line like this,
+ *         PlatformOptions.Settings = UPlatformSettings::GetAllPlatformSettings<UMyPerPlatformSettings>();
+ *         This will actually ensure that you initialize the settings exposed in the editor to whatever
+ *         the current platform configuration is for them.
+ * 
+ * Step 4) Nothing else needed.  In your system code, you will just call 
+ *         UMyPerPlatformSettings* MySettings = UPlatformSettings::GetSettingsForPlatform<UMyPerPlatformSettings>()
+ *         that will get you the current settings for the active platform, or the simulated platform in the editor.
+ */
+UCLASS(Abstract, perObjectConfig)
+class DEVELOPERSETTINGS_API UPlatformSettings : public UObject
+{
+	GENERATED_BODY()
+
+public:
+	UPlatformSettings(const FObjectInitializer& ObjectInitializer = FObjectInitializer::Get());
+	
+public:
+	template <typename TPlatformSettingsClass>
+	FORCEINLINE static TObjectPtr<TPlatformSettingsClass> GetSettingsForPlatform()
+	{
+		return Cast<TPlatformSettingsClass>(GetSettingsForPlatform(TPlatformSettingsClass::StaticClass()));
+	}
+
+	static TObjectPtr<UPlatformSettings> GetSettingsForPlatform(TSubclassOf<UPlatformSettings> SettingsClass);
+
+#if WITH_EDITOR
+	static FString GetEditorSimulatedPlatform() { return SimulatedEditorPlatform; }
+	static void SetEditorSimulatedPlatform(FString PlatformIniName) { SimulatedEditorPlatform = PlatformIniName; }
+
+	template <typename TPlatformSettingsClass>
+	FORCEINLINE static TArray<TObjectPtr<UPlatformSettings>> GetAllPlatformSettings()
+	{
+		return GetAllPlatformSettings(TPlatformSettingsClass::StaticClass());
+	}
+
+	static TArray<TObjectPtr<UPlatformSettings>> GetAllPlatformSettings(TSubclassOf<UPlatformSettings> SettingsClass);
+
+	template <typename TPlatformSettingsClass>
+	FORCEINLINE static TObjectPtr<TPlatformSettingsClass> GetSettingsForPlatform(FString TargetIniPlatformName)
+	{
+		return Cast<TPlatformSettingsClass>(GetSettingsForPlatformInternal(TPlatformSettingsClass::StaticClass(), TargetIniPlatformName));
+	}
+#endif
+
+	virtual void InitializePlatformDefaults() { }
+
+	FString GetPlatformIniName() const { return ConfigPlatformName; }
+	
+	virtual const TCHAR* GetConfigOverridePlatform() const override
+	{
+		return ConfigPlatformName.IsEmpty() ? nullptr : *ConfigPlatformName;
+	}
+
+private:
+	static UPlatformSettings* GetSettingsForPlatformInternal(TSubclassOf<UPlatformSettings> SettingsClass, FString TargetIniPlatformName);
+
+	FString ConfigPlatformName;
+
+#if WITH_EDITOR
+	static FString SimulatedEditorPlatform;
+#endif
+};
