Fixing several CIS issues with IWUY.

diff --git a/Engine/Plugins/Experimental/CommonUI/Source/CommonInput/Private/CommonInputSettings.cpp b/Engine/Plugins/Experimental/CommonUI/Source/CommonInput/Private/CommonInputSettings.cpp
--- a/Engine/Plugins/Experimental/CommonUI/Source/CommonInput/Private/CommonInputSettings.cpp
+++ b/Engine/Plugins/Experimental/CommonUI/Source/CommonInput/Private/CommonInputSettings.cpp
@@ -6,6 +6,8 @@
 #include "Engine/StreamableManager.h"
 #include "Engine/AssetManager.h"
 #include "Misc/DataDrivenPlatformInfoRegistry.h"
+#include "CommonInputBaseTypes.h"
+#include "Engine/PlatformSettings.h"
 
 UCommonInputSettings::UCommonInputSettings(const FObjectInitializer& Initializer)
 	: Super(Initializer)
@@ -112,6 +114,7 @@ void UCommonInputSettings::PostInitProperties()
 {
 	Super::PostInitProperties();
 
+#if WITH_EDITOR
 	if (CommonInputPlatformData_DEPRECATED.Num())
 	{
 		for (const auto& PlatformData : CommonInputPlatformData_DEPRECATED)
@@ -151,4 +154,5 @@ void UCommonInputSettings::PostInitProperties()
 		CommonInputPlatformData_DEPRECATED.Reset();
 		UpdateDefaultConfigFile();
 	}
+#endif
 }
diff --git a/Engine/Plugins/Experimental/CommonUI/Source/CommonUI/Private/CommonActivatableWidget.cpp b/Engine/Plugins/Experimental/CommonUI/Source/CommonUI/Private/CommonActivatableWidget.cpp
--- a/Engine/Plugins/Experimental/CommonUI/Source/CommonUI/Private/CommonActivatableWidget.cpp
+++ b/Engine/Plugins/Experimental/CommonUI/Source/CommonUI/Private/CommonActivatableWidget.cpp
@@ -14,7 +14,8 @@ void UCommonActivatableWidget::NativeConstruct()
 	if (bIsBackHandler)
 	{
 		FBindUIActionArgs BindArgs(ICommonInputModule::GetSettings().GetDefaultBackAction(), FSimpleDelegate::CreateUObject(this, &UCommonActivatableWidget::HandleBackAction));
-		BindArgs.bDisplayInActionBar = false;
+		BindArgs.bDisplayInActionBar = bIsBackActionDisplayedInActionBar;
+
 		RegisterUIActionBinding(BindArgs);
 	}
 
diff --git a/Engine/Plugins/Experimental/CommonUI/Source/CommonUI/Public/CommonActivatableWidget.h b/Engine/Plugins/Experimental/CommonUI/Source/CommonUI/Public/CommonActivatableWidget.h
--- a/Engine/Plugins/Experimental/CommonUI/Source/CommonUI/Public/CommonActivatableWidget.h
+++ b/Engine/Plugins/Experimental/CommonUI/Source/CommonUI/Public/CommonActivatableWidget.h
@@ -4,6 +4,7 @@
 
 #include "CommonUserWidget.h"
 #include "Input/UIActionBindingHandle.h"
+#include "UITag.h"
 #include "CommonActivatableWidget.generated.h"
 
 DECLARE_DYNAMIC_MULTICAST_DELEGATE(FOnWidgetActivationChanged);
@@ -117,13 +118,17 @@ protected:
 
 	void HandleBackAction();
 
-	/** True to automatically activate upon construction */
-	UPROPERTY(EditAnywhere, Category = Activation)
-	bool bAutoActivate = false;
+	/** True to receive "Back" actions automatically. Custom back handler behavior can be provided, default is to deactivate. */
+	UPROPERTY(EditAnywhere, Category = Back)
+	bool bIsBackHandler = false;
 
 	/** True to receive "Back" actions automatically. Custom back handler behavior can be provided, default is to deactivate. */
+	UPROPERTY(EditAnywhere, Category = Back)
+	bool bIsBackActionDisplayedInActionBar = false;
+
+	/** True to automatically activate upon construction */
 	UPROPERTY(EditAnywhere, Category = Activation)
-	bool bIsBackHandler = false;
+	bool bAutoActivate = false;
 
 	//@todo DanH: This property name suuucks, need to circle back and pick something better
 	/**
diff --git a/Engine/Source/Runtime/DeveloperSettings/Private/Engine/PlatformSettings.cpp b/Engine/Source/Runtime/DeveloperSettings/Private/Engine/PlatformSettings.cpp
--- a/Engine/Source/Runtime/DeveloperSettings/Private/Engine/PlatformSettings.cpp
+++ b/Engine/Source/Runtime/DeveloperSettings/Private/Engine/PlatformSettings.cpp
@@ -21,9 +21,9 @@ UPlatformSettings::UPlatformSettings(const FObjectInitializer& ObjectInitializer
 }
 
 #if WITH_EDITOR
-TArray<TObjectPtr<UPlatformSettings>> UPlatformSettings::GetAllPlatformSettings(TSubclassOf<UPlatformSettings> SettingsClass)
+TArray<UPlatformSettings*> UPlatformSettings::GetAllPlatformSettings(TSubclassOf<UPlatformSettings> SettingsClass)
 {
-	TArray<TObjectPtr<UPlatformSettings>> Settings;
+	TArray<UPlatformSettings*> Settings;
 
 	const FProjectDescriptor* Project = IProjectManager::Get().GetCurrentProject();
 	
@@ -49,7 +49,7 @@ TArray<TObjectPtr<UPlatformSettings>> UPlatformSettings::GetAllPlatformSettings(
 }
 #endif
 
-TObjectPtr<UPlatformSettings> UPlatformSettings::GetSettingsForPlatform(TSubclassOf<UPlatformSettings> SettingsClass)
+UPlatformSettings* UPlatformSettings::GetSettingsForPlatform(TSubclassOf<UPlatformSettings> SettingsClass)
 {
 	static UPlatformSettings* ThisPlatformsSettings = GetSettingsForPlatformInternal(SettingsClass, FPlatformProperties::IniPlatformName());
 	
diff --git a/Engine/Source/Runtime/DeveloperSettings/Public/Engine/PlatformSettings.h b/Engine/Source/Runtime/DeveloperSettings/Public/Engine/PlatformSettings.h
--- a/Engine/Source/Runtime/DeveloperSettings/Public/Engine/PlatformSettings.h
+++ b/Engine/Source/Runtime/DeveloperSettings/Public/Engine/PlatformSettings.h
@@ -51,27 +51,27 @@ public:
 	
 public:
 	template <typename TPlatformSettingsClass>
-	FORCEINLINE static TObjectPtr<TPlatformSettingsClass> GetSettingsForPlatform()
+	FORCEINLINE static TPlatformSettingsClass* GetSettingsForPlatform()
 	{
 		return Cast<TPlatformSettingsClass>(GetSettingsForPlatform(TPlatformSettingsClass::StaticClass()));
 	}
 
-	static TObjectPtr<UPlatformSettings> GetSettingsForPlatform(TSubclassOf<UPlatformSettings> SettingsClass);
+	static UPlatformSettings* GetSettingsForPlatform(TSubclassOf<UPlatformSettings> SettingsClass);
 
 #if WITH_EDITOR
 	static FString GetEditorSimulatedPlatform() { return SimulatedEditorPlatform; }
 	static void SetEditorSimulatedPlatform(FString PlatformIniName) { SimulatedEditorPlatform = PlatformIniName; }
 
 	template <typename TPlatformSettingsClass>
-	FORCEINLINE static TArray<TObjectPtr<UPlatformSettings>> GetAllPlatformSettings()
+	FORCEINLINE static TArray<UPlatformSettings*> GetAllPlatformSettings()
 	{
 		return GetAllPlatformSettings(TPlatformSettingsClass::StaticClass());
 	}
 
-	static TArray<TObjectPtr<UPlatformSettings>> GetAllPlatformSettings(TSubclassOf<UPlatformSettings> SettingsClass);
+	static TArray<UPlatformSettings*> GetAllPlatformSettings(TSubclassOf<UPlatformSettings> SettingsClass);
 
 	template <typename TPlatformSettingsClass>
-	FORCEINLINE static TObjectPtr<TPlatformSettingsClass> GetSettingsForPlatform(FString TargetIniPlatformName)
+	FORCEINLINE static TPlatformSettingsClass* GetSettingsForPlatform(FString TargetIniPlatformName)
 	{
 		return Cast<TPlatformSettingsClass>(GetSettingsForPlatformInternal(TPlatformSettingsClass::StaticClass(), TargetIniPlatformName));
 	}
