PlatformSettings - Adding a common utility function for getting the platform inis that are applicable to the running project.

diff --git a/Engine/Source/Runtime/DeveloperSettings/Private/Engine/PlatformSettings.cpp b/Engine/Source/Runtime/DeveloperSettings/Private/Engine/PlatformSettings.cpp
--- a/Engine/Source/Runtime/DeveloperSettings/Private/Engine/PlatformSettings.cpp
+++ b/Engine/Source/Runtime/DeveloperSettings/Private/Engine/PlatformSettings.cpp
@@ -11,39 +11,55 @@
 #include "ProjectDescriptor.h"
 
 #if WITH_EDITOR
-FString UPlatformSettings::SimulatedEditorPlatform;
+FName UPlatformSettings::SimulatedEditorPlatform;
 #endif
 
 UPlatformSettings::UPlatformSettings(const FObjectInitializer& ObjectInitializer)
 	: Super(ObjectInitializer)
 {
-	
+
 }
 
 #if WITH_EDITOR
-TArray<UPlatformSettings*> UPlatformSettings::GetAllPlatformSettings(TSubclassOf<UPlatformSettings> SettingsClass)
+TArray<FName> UPlatformSettings::GetKnownAndEnablePlatformIniNames()
 {
-	TArray<UPlatformSettings*> Settings;
+	TArray<FName> Results;
 
-	const FProjectDescriptor* Project = IProjectManager::Get().GetCurrentProject();
-	
-	const TMap<FName, FDataDrivenPlatformInfo>& AllPlatforms = FDataDrivenPlatformInfoRegistry::GetAllPlatformInfos();
+	FProjectStatus ProjectStatus;
+	const bool bProjectStatusIsValid = IProjectManager::Get().QueryStatusForCurrentProject(/*out*/ ProjectStatus);
 
-	for (auto& PlatformKVP : AllPlatforms)
+	for (const auto& Pair : FDataDrivenPlatformInfoRegistry::GetAllPlatformInfos())
 	{
-		const FName PlatformIniName = PlatformKVP.Key;
-		const FDataDrivenPlatformInfo& PlatformInfo = PlatformKVP.Value;
+		const FName PlatformName = Pair.Key;
+		const FDataDrivenPlatformInfo& Info = Pair.Value;
 
-		if (PlatformInfo.bIsFakePlatform)
-		{
-			continue;
-		}
+		const bool bProjectDisabledPlatform = bProjectStatusIsValid && !ProjectStatus.IsTargetPlatformSupported(PlatformName);
+
+		const bool bEnabledForUse =
+#if DDPI_HAS_EXTENDED_PLATFORMINFO_DATA
+			Info.bEnabledForUse;
+#else
+			true;
+#endif
 
-		if (Project->TargetPlatforms.IsEmpty() || Project->TargetPlatforms.Contains(PlatformIniName))
+		const bool bSupportedPlatform = !Info.bIsFakePlatform && bEnabledForUse && !bProjectDisabledPlatform;
+
+		if (bSupportedPlatform)
 		{
-			Settings.Add(GetSettingsForPlatformInternal(SettingsClass, PlatformIniName.ToString()));
+			Results.Add(PlatformName);
 		}
 	}
+
+	return Results;
+}
+
+TArray<UPlatformSettings*> UPlatformSettings::GetAllPlatformSettings(TSubclassOf<UPlatformSettings> SettingsClass)
+{
+	TArray<UPlatformSettings*> Settings;
+	for (FName PlatformIniName : GetKnownAndEnablePlatformIniNames())
+	{
+		Settings.Add(GetSettingsForPlatformInternal(SettingsClass, PlatformIniName.ToString()));
+	}
 	
 	return Settings;
 }
@@ -52,11 +68,15 @@ TArray<UPlatformSettings*> UPlatformSettings::GetAllPlatformSettings(TSubclassOf
 UPlatformSettings* UPlatformSettings::GetSettingsForPlatform(TSubclassOf<UPlatformSettings> SettingsClass)
 {
 	static UPlatformSettings* ThisPlatformsSettings = GetSettingsForPlatformInternal(SettingsClass, FPlatformProperties::IniPlatformName());
-	
+
 #if WITH_EDITOR
-	if (GIsEditor && !SimulatedEditorPlatform.IsEmpty())
+	if (GIsEditor && SimulatedEditorPlatform != NAME_None)
 	{
-		return GetSettingsForPlatformInternal(SettingsClass, SimulatedEditorPlatform);
+		UPlatformSettings* OtherPlatformsSettings = GetSettingsForPlatformInternal(SettingsClass, SimulatedEditorPlatform.ToString());
+		if (ensure(OtherPlatformsSettings))
+		{
+			return OtherPlatformsSettings;
+		}
 	}
 #endif
 	
diff --git a/Engine/Source/Runtime/DeveloperSettings/Public/Engine/PlatformSettings.h b/Engine/Source/Runtime/DeveloperSettings/Public/Engine/PlatformSettings.h
--- a/Engine/Source/Runtime/DeveloperSettings/Public/Engine/PlatformSettings.h
+++ b/Engine/Source/Runtime/DeveloperSettings/Public/Engine/PlatformSettings.h
@@ -59,8 +59,10 @@ public:
 	static UPlatformSettings* GetSettingsForPlatform(TSubclassOf<UPlatformSettings> SettingsClass);
 
 #if WITH_EDITOR
-	static FString GetEditorSimulatedPlatform() { return SimulatedEditorPlatform; }
-	static void SetEditorSimulatedPlatform(FString PlatformIniName) { SimulatedEditorPlatform = PlatformIniName; }
+	static FName GetEditorSimulatedPlatform() { return SimulatedEditorPlatform; }
+	static void SetEditorSimulatedPlatform(FName PlatformIniName) { SimulatedEditorPlatform = PlatformIniName; }
+
+	static TArray<FName> GetKnownAndEnablePlatformIniNames();
 
 	template <typename TPlatformSettingsClass>
 	FORCEINLINE static TArray<UPlatformSettings*> GetAllPlatformSettings()
@@ -92,6 +94,6 @@ private:
 	FString ConfigPlatformName;
 
 #if WITH_EDITOR
-	static FString SimulatedEditorPlatform;
+	static FName SimulatedEditorPlatform;
 #endif
 };
