Prevent a crash when there is more than one kind of UPlatformSettings usage in a project
- Removed the problematic static that caused the crash and introduced a new caching system via a UPlatformSettingsManager class that keeps explicit references to the settings instead of adding them to the root
- Switched the platform names to be more consistently FName instead of roundtripping through FString

- Slightly unrelated: Added DeterminesOutputType metadata to FindParentWidgetOfType to avoid needing an extra cast in typical usage

diff --git a/Engine/Plugins/Experimental/CommonUI/Source/CommonInput/Private/CommonInputBaseTypes.cpp b/Engine/Plugins/Experimental/CommonUI/Source/CommonInput/Private/CommonInputBaseTypes.cpp
--- a/Engine/Plugins/Experimental/CommonUI/Source/CommonInput/Private/CommonInputBaseTypes.cpp
+++ b/Engine/Plugins/Experimental/CommonUI/Source/CommonInput/Private/CommonInputBaseTypes.cpp
@@ -246,7 +246,7 @@ void UCommonInputPlatformSettings::InitializeControllerData() const
 
 void UCommonInputPlatformSettings::InitializePlatformDefaults()
 {
-	const FName PlatformName = *GetPlatformIniName();
+	const FName PlatformName = GetPlatformIniName();
 	const FDataDrivenPlatformInfo& PlatformInfo = FDataDrivenPlatformInfoRegistry::GetPlatformInfo(PlatformName);
 	if (PlatformInfo.DefaultInputType == "Gamepad")
 	{
diff --git a/Engine/Plugins/Experimental/CommonUI/Source/CommonInput/Private/CommonInputSettings.cpp b/Engine/Plugins/Experimental/CommonUI/Source/CommonInput/Private/CommonInputSettings.cpp
--- a/Engine/Plugins/Experimental/CommonUI/Source/CommonInput/Private/CommonInputSettings.cpp
+++ b/Engine/Plugins/Experimental/CommonUI/Source/CommonInput/Private/CommonInputSettings.cpp
@@ -14,9 +14,7 @@ UCommonInputSettings::UCommonInputSettings(const FObjectInitializer& Initializer
 	: Super(Initializer)
 	, bInputDataLoaded(false)
 {
-#if WITH_EDITOR
-	PlatformInput.Settings = UPlatformSettings::GetAllPlatformSettings<UCommonInputPlatformSettings>();
-#endif
+	PlatformInput.Initialize(UCommonInputPlatformSettings::StaticClass());
 }
 
 void UCommonInputSettings::LoadData()
@@ -146,7 +144,7 @@ void UCommonInputSettings::PostInitProperties()
 		{
 			const FCommonInputPlatformBaseData& OriginalData = PlatformData.Value;
 
-			if (UCommonInputPlatformSettings* Settings = UPlatformSettings::GetSettingsForPlatform<UCommonInputPlatformSettings>(PlatformData.Key.ToString()))
+			if (UCommonInputPlatformSettings* Settings = UPlatformSettingsManager::Get().GetSettingsForPlatform<UCommonInputPlatformSettings>(PlatformData.Key))
 			{
 				Settings->bSupportsMouseAndKeyboard = OriginalData.bSupportsMouseAndKeyboard;
 				Settings->bSupportsGamepad = OriginalData.bSupportsGamepad;
@@ -161,9 +159,9 @@ void UCommonInputSettings::PostInitProperties()
 			else if (PlatformData.Key == FCommonInputDefaults::PlatformPC)
 			{
 				TArray<UCommonInputPlatformSettings*> PCPlatforms;
-				PCPlatforms.Add(UPlatformSettings::GetSettingsForPlatform<UCommonInputPlatformSettings>("Windows"));
-				PCPlatforms.Add(UPlatformSettings::GetSettingsForPlatform<UCommonInputPlatformSettings>("WinGDK"));
-				PCPlatforms.Add(UPlatformSettings::GetSettingsForPlatform<UCommonInputPlatformSettings>("Linux"));
+				PCPlatforms.Add(UPlatformSettingsManager::Get().GetSettingsForPlatform<UCommonInputPlatformSettings>("Windows"));
+				PCPlatforms.Add(UPlatformSettingsManager::Get().GetSettingsForPlatform<UCommonInputPlatformSettings>("WinGDK"));
+				PCPlatforms.Add(UPlatformSettingsManager::Get().GetSettingsForPlatform<UCommonInputPlatformSettings>("Linux"));
 
 				for (UCommonInputPlatformSettings* PCPlatform : PCPlatforms)
 				{
diff --git a/Engine/Plugins/Experimental/CommonUI/Source/CommonInput/Private/CommonInputSubsystem.cpp b/Engine/Plugins/Experimental/CommonUI/Source/CommonInput/Private/CommonInputSubsystem.cpp
--- a/Engine/Plugins/Experimental/CommonUI/Source/CommonInput/Private/CommonInputSubsystem.cpp
+++ b/Engine/Plugins/Experimental/CommonUI/Source/CommonInput/Private/CommonInputSubsystem.cpp
@@ -314,8 +314,11 @@ void UCommonInputSubsystem::Initialize(FSubsystemCollectionBase& Collection)
 	Super::Initialize(Collection);
 
 	FCommonInputBase::GetInputSettings()->LoadData();
-	GamepadInputType = UPlatformSettings::GetSettingsForPlatform<UCommonInputPlatformSettings>()->GetDefaultGamepadName();
-	CurrentInputType = LastInputType = UPlatformSettings::GetSettingsForPlatform<UCommonInputPlatformSettings>()->GetDefaultInputType();
+
+	const UCommonInputPlatformSettings* Settings = UPlatformSettingsManager::Get().GetSettingsForPlatform<UCommonInputPlatformSettings>();
+
+	GamepadInputType = Settings->GetDefaultGamepadName();
+	CurrentInputType = LastInputType = Settings->GetDefaultInputType();
 
 	CommonInputPreprocessor = MakeShared<FCommonInputPreprocessor>(*this);
 	FSlateApplication::Get().RegisterInputPreProcessor(CommonInputPreprocessor, 0);
diff --git a/Engine/Plugins/Experimental/CommonUI/Source/CommonInput/Public/CommonInputBaseTypes.h b/Engine/Plugins/Experimental/CommonUI/Source/CommonInput/Public/CommonInputBaseTypes.h
--- a/Engine/Plugins/Experimental/CommonUI/Source/CommonInput/Public/CommonInputBaseTypes.h
+++ b/Engine/Plugins/Experimental/CommonUI/Source/CommonInput/Public/CommonInputBaseTypes.h
@@ -177,7 +177,7 @@ public:
 
 	static UCommonInputPlatformSettings* Get()
 	{
-		return UPlatformSettings::GetSettingsForPlatform<UCommonInputPlatformSettings>();
+		return UPlatformSettingsManager::Get().GetSettingsForPlatform<UCommonInputPlatformSettings>();
 	}
 
 	bool TryGetInputBrush(FSlateBrush& OutBrush, FKey Key, ECommonInputType InputType, const FName GamepadName) const;
diff --git a/Engine/Plugins/Experimental/CommonUI/Source/CommonUI/Private/CommonUISettings.cpp b/Engine/Plugins/Experimental/CommonUI/Source/CommonUI/Private/CommonUISettings.cpp
--- a/Engine/Plugins/Experimental/CommonUI/Source/CommonUI/Private/CommonUISettings.cpp
+++ b/Engine/Plugins/Experimental/CommonUI/Source/CommonUI/Private/CommonUISettings.cpp
@@ -132,7 +132,7 @@ UObject* UCommonUISettings::GetDefaultImageResourceObject() const
 const FGameplayTagContainer& UCommonUISettings::GetPlatformTraits() const
 {
 #if WITH_EDITOR
-	const FName SimulatedPlatform = UPlatformSettings::GetEditorSimulatedPlatform();
+	const FName SimulatedPlatform = UPlatformSettingsManager::GetEditorSimulatedPlatform();
 	if (SimulatedPlatform != NAME_None)
 	{
 		static FName LastSimulatedPlatform = NAME_None;
diff --git a/Engine/Plugins/Experimental/CommonUI/Source/CommonUI/Public/CommonUILibrary.h b/Engine/Plugins/Experimental/CommonUI/Source/CommonUI/Public/CommonUILibrary.h
--- a/Engine/Plugins/Experimental/CommonUI/Source/CommonUI/Public/CommonUILibrary.h
+++ b/Engine/Plugins/Experimental/CommonUI/Source/CommonUI/Public/CommonUILibrary.h
@@ -18,6 +18,6 @@ public:
 	/**
 	 * Finds the first parent widget of the given type and returns it, or null if no parent could be found.
 	 */
-	UFUNCTION(BlueprintCallable, Category="Common UI")
+	UFUNCTION(BlueprintCallable, Category="Common UI", meta=(DeterminesOutputType=Type))
 	static UWidget* FindParentWidgetOfType(UWidget* StartingWidget, TSubclassOf<UWidget> Type);
 };
diff --git a/Engine/Source/Developer/SettingsEditor/Private/Customizations/PerPlatformSettingsCustomization.cpp b/Engine/Source/Developer/SettingsEditor/Private/Customizations/PerPlatformSettingsCustomization.cpp
--- a/Engine/Source/Developer/SettingsEditor/Private/Customizations/PerPlatformSettingsCustomization.cpp
+++ b/Engine/Source/Developer/SettingsEditor/Private/Customizations/PerPlatformSettingsCustomization.cpp
@@ -52,7 +52,7 @@ void FPerPlatformSettingsCustomization::CustomizeChildren(TSharedRef<IPropertyHa
 						.AllowChildren(true)
 						.CreateCategoryNodes(false);
 
-					IDetailGroup& PlatformGroup = ChildBuilder.AddGroup(PlatformSettings->GetFName(), FText::FromString(PlatformSettings->GetPlatformIniName()));
+					IDetailGroup& PlatformGroup = ChildBuilder.AddGroup(PlatformSettings->GetFName(), FText::FromName(PlatformSettings->GetPlatformIniName()));
 
 					TSharedPtr<IPropertyHandle> PlatformSettingsPropertyArrayEntry = PlatformSettingsProperty->GetChildHandle(0);
 
diff --git a/Engine/Source/Runtime/DeveloperSettings/Private/Engine/PlatformSettings.cpp b/Engine/Source/Runtime/DeveloperSettings/Private/Engine/PlatformSettings.cpp
--- a/Engine/Source/Runtime/DeveloperSettings/Private/Engine/PlatformSettings.cpp
+++ b/Engine/Source/Runtime/DeveloperSettings/Private/Engine/PlatformSettings.cpp
@@ -1,111 +1,15 @@
 // Copyright Epic Games, Inc. All Rights Reserved.
 
 #include "Engine/PlatformSettings.h"
-#include "HAL/IConsoleManager.h"
-#include "UObject/UnrealType.h"
-#include "UObject/EnumProperty.h"
-#include "UObject/PropertyPortFlags.h"
-#include "Misc/DataDrivenPlatformInfoRegistry.h"
-#include "HAL/PlatformProperties.h"
-#include "Interfaces/IProjectManager.h"
-#include "ProjectDescriptor.h"
 
-#if WITH_EDITOR
-FName UPlatformSettings::SimulatedEditorPlatform;
-#endif
-
-UPlatformSettings::UPlatformSettings(const FObjectInitializer& ObjectInitializer)
-	: Super(ObjectInitializer)
-{
-
-}
-
-#if WITH_EDITOR
-TArray<FName> UPlatformSettings::GetKnownAndEnablePlatformIniNames()
-{
-	TArray<FName> Results;
-
-	FProjectStatus ProjectStatus;
-	const bool bProjectStatusIsValid = IProjectManager::Get().QueryStatusForCurrentProject(/*out*/ ProjectStatus);
-
-	for (const auto& Pair : FDataDrivenPlatformInfoRegistry::GetAllPlatformInfos())
-	{
-		const FName PlatformName = Pair.Key;
-		const FDataDrivenPlatformInfo& Info = Pair.Value;
-
-		const bool bProjectDisabledPlatform = bProjectStatusIsValid && !ProjectStatus.IsTargetPlatformSupported(PlatformName);
-
-		const bool bEnabledForUse =
-#if DDPI_HAS_EXTENDED_PLATFORMINFO_DATA
-			Info.bEnabledForUse;
-#else
-			true;
-#endif
-
-		const bool bSupportedPlatform = !Info.bIsFakePlatform && bEnabledForUse && !bProjectDisabledPlatform;
-
-		if (bSupportedPlatform)
-		{
-			Results.Add(PlatformName);
-		}
-	}
-
-	return Results;
-}
-
-TArray<UPlatformSettings*> UPlatformSettings::GetAllPlatformSettings(TSubclassOf<UPlatformSettings> SettingsClass)
+void FPerPlatformSettings::Initialize(TSubclassOf<UPlatformSettings> SettingsClass)
 {
-	TArray<UPlatformSettings*> Settings;
-	for (FName PlatformIniName : GetKnownAndEnablePlatformIniNames())
-	{
-		Settings.Add(GetSettingsForPlatformInternal(SettingsClass, PlatformIniName.ToString()));
-	}
-	
-	return Settings;
-}
-#endif
-
-UPlatformSettings* UPlatformSettings::GetSettingsForPlatform(TSubclassOf<UPlatformSettings> SettingsClass)
-{
-	static UPlatformSettings* ThisPlatformsSettings = GetSettingsForPlatformInternal(SettingsClass, FPlatformProperties::IniPlatformName());
-
 #if WITH_EDITOR
-	if (GIsEditor && SimulatedEditorPlatform != NAME_None)
-	{
-		UPlatformSettings* OtherPlatformsSettings = GetSettingsForPlatformInternal(SettingsClass, SimulatedEditorPlatform.ToString());
-		if (ensure(OtherPlatformsSettings))
-		{
-			return OtherPlatformsSettings;
-		}
-	}
+	Settings = UPlatformSettingsManager::Get().GetAllPlatformSettings(SettingsClass);
 #endif
-	
-	return ThisPlatformsSettings;
 }
 
-UPlatformSettings* UPlatformSettings::GetSettingsForPlatformInternal(TSubclassOf<UPlatformSettings> SettingsClass, FString TargetIniPlatformName)
+UPlatformSettings::UPlatformSettings(const FObjectInitializer& ObjectInitializer)
+	: Super(ObjectInitializer)
 {
-	if (!ensure(SettingsClass.Get()))
-	{
-		return nullptr;
-	}
-
-	const FString PlatformSettingsName = SettingsClass->GetName() + TEXT("_") + TargetIniPlatformName;
-	
-	UPlatformSettings* PlatformSettingsForThisClass = FindObject<UPlatformSettings>(GetTransientPackage(), *PlatformSettingsName);
-	if (PlatformSettingsForThisClass == nullptr)
-	{
-		const TMap<FName, FDataDrivenPlatformInfo>& AllPlatforms = FDataDrivenPlatformInfoRegistry::GetAllPlatformInfos();
-		if (AllPlatforms.Contains(*TargetIniPlatformName))
-		{
-			PlatformSettingsForThisClass = NewObject<UPlatformSettings>(GetTransientPackage(), SettingsClass, FName(*PlatformSettingsName));
-			PlatformSettingsForThisClass->ConfigPlatformName = TargetIniPlatformName;
-			PlatformSettingsForThisClass->AddToRoot();
-			PlatformSettingsForThisClass->InitializePlatformDefaults();
-
-			PlatformSettingsForThisClass->LoadConfig();
-		}
-	}
-	
-	return PlatformSettingsForThisClass;
 }
diff --git a/Engine/Source/Runtime/DeveloperSettings/Private/Engine/PlatformSettingsManager.cpp b/Engine/Source/Runtime/DeveloperSettings/Private/Engine/PlatformSettingsManager.cpp
new file mode 100644
--- /dev/null
+++ b/Engine/Source/Runtime/DeveloperSettings/Private/Engine/PlatformSettingsManager.cpp
@@ -0,0 +1,111 @@
+// Copyright Epic Games, Inc. All Rights Reserved.
+
+#include "Engine/PlatformSettingsManager.h"
+#include "Engine/PlatformSettings.h"
+#include "HAL/IConsoleManager.h"
+#include "UObject/UnrealType.h"
+#include "UObject/EnumProperty.h"
+#include "UObject/PropertyPortFlags.h"
+#include "Misc/DataDrivenPlatformInfoRegistry.h"
+#include "HAL/PlatformProperties.h"
+#include "Interfaces/IProjectManager.h"
+#include "ProjectDescriptor.h"
+
+#if WITH_EDITOR
+FName UPlatformSettingsManager::SimulatedEditorPlatform;
+#endif
+
+UPlatformSettingsManager::UPlatformSettingsManager(const FObjectInitializer& ObjectInitializer)
+	: Super(ObjectInitializer)
+	, IniPlatformName(FPlatformProperties::IniPlatformName())
+{
+}
+
+UPlatformSettings* UPlatformSettingsManager::GetSettingsForPlatform(TSubclassOf<UPlatformSettings> SettingsClass) const
+{
+#if WITH_EDITOR
+	if (GIsEditor && !SimulatedEditorPlatform.IsNone())
+	{
+		return GetSettingsForPlatform(SettingsClass, SimulatedEditorPlatform);
+	}
+#endif
+
+	return GetSettingsForPlatform(SettingsClass, IniPlatformName);
+}
+
+UPlatformSettings* UPlatformSettingsManager::GetSettingsForPlatform(TSubclassOf<UPlatformSettings> SettingsClass, FName DesiredPlatformIniName) const
+{
+	check(SettingsClass != nullptr);
+
+	FPlatformSettingsInstances& Record = SettingsMap.FindOrAdd(SettingsClass);
+
+	TObjectPtr<UPlatformSettings>& Result = (DesiredPlatformIniName == IniPlatformName) ? Record.PlatformInstance : Record.OtherPlatforms.FindOrAdd(DesiredPlatformIniName);
+	if (Result == nullptr)
+	{
+		Result = CreateSettingsObjectForPlatform(SettingsClass, DesiredPlatformIniName);
+	}
+	return Result;
+}
+
+UPlatformSettings* UPlatformSettingsManager::CreateSettingsObjectForPlatform(TSubclassOf<UPlatformSettings> SettingsClass, FName TargetIniPlatformName) const
+{
+	const FString PlatformSettingsName = SettingsClass->GetName() + TEXT("_") + TargetIniPlatformName.ToString();
+	UPlatformSettings* PlatformSettingsForThisClass = NewObject<UPlatformSettings>(GetTransientPackage(), SettingsClass, FName(*PlatformSettingsName));
+	PlatformSettingsForThisClass->ConfigPlatformName = TargetIniPlatformName;
+	PlatformSettingsForThisClass->ConfigPlatformNameStr = TargetIniPlatformName.ToString();
+
+	const TMap<FName, FDataDrivenPlatformInfo>& AllPlatforms = FDataDrivenPlatformInfoRegistry::GetAllPlatformInfos();
+	if (ensureMsgf(AllPlatforms.Contains(TargetIniPlatformName), TEXT("Platform config %s requested for unknown platform %s"), *SettingsClass->GetPathName(), *TargetIniPlatformName.ToString()))
+	{
+		PlatformSettingsForThisClass->InitializePlatformDefaults();
+
+		PlatformSettingsForThisClass->LoadConfig();
+	}
+	
+	return PlatformSettingsForThisClass;
+}
+
+#if WITH_EDITOR
+TArray<FName> UPlatformSettingsManager::GetKnownAndEnablePlatformIniNames()
+{
+	TArray<FName> Results;
+
+	FProjectStatus ProjectStatus;
+	const bool bProjectStatusIsValid = IProjectManager::Get().QueryStatusForCurrentProject(/*out*/ ProjectStatus);
+
+	for (const auto& Pair : FDataDrivenPlatformInfoRegistry::GetAllPlatformInfos())
+	{
+		const FName PlatformName = Pair.Key;
+		const FDataDrivenPlatformInfo& Info = Pair.Value;
+
+		const bool bProjectDisabledPlatform = bProjectStatusIsValid && !ProjectStatus.IsTargetPlatformSupported(PlatformName);
+
+		const bool bEnabledForUse =
+#if DDPI_HAS_EXTENDED_PLATFORMINFO_DATA
+			Info.bEnabledForUse;
+#else
+			true;
+#endif
+
+		const bool bSupportedPlatform = !Info.bIsFakePlatform && bEnabledForUse && !bProjectDisabledPlatform;
+
+		if (bSupportedPlatform)
+		{
+			Results.Add(PlatformName);
+		}
+	}
+
+	return Results;
+}
+
+TArray<UPlatformSettings*> UPlatformSettingsManager::GetAllPlatformSettings(TSubclassOf<UPlatformSettings> SettingsClass) const
+{
+	TArray<UPlatformSettings*> Settings;
+	for (FName PlatformIniName : GetKnownAndEnablePlatformIniNames())
+	{
+		Settings.Add(GetSettingsForPlatform(SettingsClass, PlatformIniName));
+	}
+
+	return Settings;
+}
+#endif
diff --git a/Engine/Source/Runtime/DeveloperSettings/Public/Engine/PlatformSettings.h b/Engine/Source/Runtime/DeveloperSettings/Public/Engine/PlatformSettings.h
--- a/Engine/Source/Runtime/DeveloperSettings/Public/Engine/PlatformSettings.h
+++ b/Engine/Source/Runtime/DeveloperSettings/Public/Engine/PlatformSettings.h
@@ -8,17 +8,24 @@
 #include "UObject/Object.h"
 #include "Templates/SubclassOf.h"
 #include "DeveloperSettings.h"
+#include "PlatformSettingsManager.h"
 
 #include "PlatformSettings.generated.h"
 
+class UPlatformSettings;
+
 USTRUCT()
-struct FPerPlatformSettings
+struct DEVELOPERSETTINGS_API FPerPlatformSettings
 {
 	GENERATED_BODY()
 
-public:
+	void Initialize(TSubclassOf<UPlatformSettings> SettingsClass);
+
+private:
 	UPROPERTY(Instanced, Transient, EditAnywhere, EditFixedSize, Category = Layout)
 	TArray<TObjectPtr<UPlatformSettings>> Settings;
+
+	friend class FPerPlatformSettingsCustomization;
 };
 
 /**
@@ -33,12 +40,12 @@ public:
  *         FPerPlatformSettings PlatformOptions
  * 
  * Step 3) In your UDeveloperSettings subclasses construct, there should be a line like this,
- *         PlatformOptions.Settings = UPlatformSettings::GetAllPlatformSettings<UMyPerPlatformSettings>();
+ *         PlatformOptions.Settings.Initialize(UMyPerPlatformSettings::StaticClass());
  *         This will actually ensure that you initialize the settings exposed in the editor to whatever
  *         the current platform configuration is for them.
  * 
  * Step 4) Nothing else needed.  In your system code, you will just call 
- *         UMyPerPlatformSettings* MySettings = UPlatformSettings::GetSettingsForPlatform<UMyPerPlatformSettings>()
+ *         UMyPerPlatformSettings* MySettings = UPlatformSettingsManager::Get().GetSettingsForPlatform<UMyPerPlatformSettings>()
  *         that will get you the current settings for the active platform, or the simulated platform in the editor.
  */
 UCLASS(Abstract, perObjectConfig)
@@ -48,52 +55,22 @@ class DEVELOPERSETTINGS_API UPlatformSettings : public UObject
 
 public:
 	UPlatformSettings(const FObjectInitializer& ObjectInitializer = FObjectInitializer::Get());
-	
-public:
-	template <typename TPlatformSettingsClass>
-	FORCEINLINE static TPlatformSettingsClass* GetSettingsForPlatform()
-	{
-		return Cast<TPlatformSettingsClass>(GetSettingsForPlatform(TPlatformSettingsClass::StaticClass()));
-	}
-
-	static UPlatformSettings* GetSettingsForPlatform(TSubclassOf<UPlatformSettings> SettingsClass);
-
-#if WITH_EDITOR
-	static FName GetEditorSimulatedPlatform() { return SimulatedEditorPlatform; }
-	static void SetEditorSimulatedPlatform(FName PlatformIniName) { SimulatedEditorPlatform = PlatformIniName; }
-
-	static TArray<FName> GetKnownAndEnablePlatformIniNames();
-
-	template <typename TPlatformSettingsClass>
-	FORCEINLINE static TArray<UPlatformSettings*> GetAllPlatformSettings()
-	{
-		return GetAllPlatformSettings(TPlatformSettingsClass::StaticClass());
-	}
-
-	static TArray<UPlatformSettings*> GetAllPlatformSettings(TSubclassOf<UPlatformSettings> SettingsClass);
-
-	template <typename TPlatformSettingsClass>
-	FORCEINLINE static TPlatformSettingsClass* GetSettingsForPlatform(FString TargetIniPlatformName)
-	{
-		return Cast<TPlatformSettingsClass>(GetSettingsForPlatformInternal(TPlatformSettingsClass::StaticClass(), TargetIniPlatformName));
-	}
-#endif
 
 	virtual void InitializePlatformDefaults() { }
 
-	FString GetPlatformIniName() const { return ConfigPlatformName; }
+	FName GetPlatformIniName() const { return ConfigPlatformName; }
 	
+	//~UObject interface
 	virtual const TCHAR* GetConfigOverridePlatform() const override
 	{
-		return ConfigPlatformName.IsEmpty() ? nullptr : *ConfigPlatformName;
+		return ConfigPlatformNameStr.IsEmpty() ? nullptr : *ConfigPlatformNameStr;
 	}
+	//~End of UObject interface
 
 private:
-	static UPlatformSettings* GetSettingsForPlatformInternal(TSubclassOf<UPlatformSettings> SettingsClass, FString TargetIniPlatformName);
-
-	FString ConfigPlatformName;
+	// The platform we are an instance of settings for
+	FName ConfigPlatformName;
+	FString ConfigPlatformNameStr;
 
-#if WITH_EDITOR
-	static FName SimulatedEditorPlatform;
-#endif
+	friend UPlatformSettingsManager;
 };
diff --git a/Engine/Source/Runtime/DeveloperSettings/Public/Engine/PlatformSettingsManager.h b/Engine/Source/Runtime/DeveloperSettings/Public/Engine/PlatformSettingsManager.h
new file mode 100644
--- /dev/null
+++ b/Engine/Source/Runtime/DeveloperSettings/Public/Engine/PlatformSettingsManager.h
@@ -0,0 +1,103 @@
+// Copyright Epic Games, Inc. All Rights Reserved.
+
+#pragma once
+
+#include "CoreMinimal.h"
+#include "UObject/ObjectMacros.h"
+#include "UObject/UObjectGlobals.h"
+#include "UObject/Object.h"
+#include "Templates/SubclassOf.h"
+#include "DeveloperSettings.h"
+
+#include "PlatformSettingsManager.generated.h"
+
+class UPlatformSettings;
+
+// List of platform-specific instances for a class
+USTRUCT()
+struct FPlatformSettingsInstances
+{
+	GENERATED_BODY()
+
+	// The instance for the native platform
+	UPROPERTY(Transient)
+	TObjectPtr<UPlatformSettings> PlatformInstance = nullptr;
+
+	// Instances for other platforms (only used in the editor)
+	UPROPERTY(Transient)
+	TMap<FName, TObjectPtr<UPlatformSettings>> OtherPlatforms;
+};
+
+// The manager for all platform-specific settings
+UCLASS()
+class DEVELOPERSETTINGS_API UPlatformSettingsManager : public UObject
+{
+	GENERATED_BODY()
+
+public:
+	UPlatformSettingsManager(const FObjectInitializer& ObjectInitializer = FObjectInitializer::Get());
+
+	static const UPlatformSettingsManager& Get()
+	{
+		return *GetDefault<UPlatformSettingsManager>();
+	}
+
+	// Returns the platform settings for the specified class for the current platform
+	UPlatformSettings* GetSettingsForPlatform(TSubclassOf<UPlatformSettings> SettingsClass) const;
+
+	// Returns the platform settings for the specified class for the specified platform
+	UPlatformSettings* GetSettingsForPlatform(TSubclassOf<UPlatformSettings> SettingsClass, FName DesiredPlatformIniName) const;
+
+#if WITH_EDITOR
+	// Gets the simulated platform
+	static FName GetEditorSimulatedPlatform() { return SimulatedEditorPlatform; }
+
+	// Sets the current simulated platform
+	static void SetEditorSimulatedPlatform(FName PlatformIniName) { SimulatedEditorPlatform = PlatformIniName; }
+#endif
+
+
+public:
+	template <typename TPlatformSettingsClass>
+	FORCEINLINE TPlatformSettingsClass* GetSettingsForPlatform() const
+	{
+		return CastChecked<TPlatformSettingsClass>(GetSettingsForPlatform(TPlatformSettingsClass::StaticClass()));
+	}
+
+#if WITH_EDITOR
+	static TArray<FName> GetKnownAndEnablePlatformIniNames();
+
+	template <typename TPlatformSettingsClass>
+	FORCEINLINE TArray<UPlatformSettings*> GetAllPlatformSettings() const
+	{
+		return GetAllPlatformSettings(TPlatformSettingsClass::StaticClass());
+	}
+
+	TArray<UPlatformSettings*> GetAllPlatformSettings(TSubclassOf<UPlatformSettings> SettingsClass) const;
+
+	template <typename TPlatformSettingsClass>
+	FORCEINLINE TPlatformSettingsClass* GetSettingsForPlatform(FName TargetIniPlatformName) const
+	{
+		return Cast<TPlatformSettingsClass>(GetSettingsForPlatform(TPlatformSettingsClass::StaticClass(), TargetIniPlatformName));
+	}
+
+#endif
+
+private:
+	// Creates a settings object for the specified class and platform
+	UPlatformSettings* CreateSettingsObjectForPlatform(TSubclassOf<UPlatformSettings> SettingsClass, FName TargetIniPlatformName) const;
+
+
+private:
+	// Created platform-specific settings
+	UPROPERTY(Transient)
+	mutable TMap<TSubclassOf<UPlatformSettings>, FPlatformSettingsInstances> SettingsMap;
+
+	// The true platform name for the current platform (unaffected by previews, etc...)
+	FName IniPlatformName;
+
+#if WITH_EDITOR
+	// Current simulated platform
+	static FName SimulatedEditorPlatform;
+#endif
+};
