Fixed a 'disregard for GC' violation in cooked games by adding the per-platform platform settings objects to the root set once again

diff --git a/Engine/Source/Runtime/DeveloperSettings/Private/Engine/PlatformSettingsManager.cpp b/Engine/Source/Runtime/DeveloperSettings/Private/Engine/PlatformSettingsManager.cpp
--- a/Engine/Source/Runtime/DeveloperSettings/Private/Engine/PlatformSettingsManager.cpp
+++ b/Engine/Source/Runtime/DeveloperSettings/Private/Engine/PlatformSettingsManager.cpp
@@ -51,6 +51,7 @@ UPlatformSettings* UPlatformSettingsManager::CreateSettingsObjectForPlatform(TSu
 {
 	const FString PlatformSettingsName = SettingsClass->GetName() + TEXT("_") + TargetIniPlatformName.ToString();
 	UPlatformSettings* PlatformSettingsForThisClass = NewObject<UPlatformSettings>(GetTransientPackage(), SettingsClass, FName(*PlatformSettingsName));
+	PlatformSettingsForThisClass->AddToRoot();
 	PlatformSettingsForThisClass->ConfigPlatformName = TargetIniPlatformName;
 	PlatformSettingsForThisClass->ConfigPlatformNameStr = TargetIniPlatformName.ToString();
 
