PR #9573: Platform Settings: fix settings object being recreated on each GetSettingsForPlatform call (Contributed by [REDACTED_PERSON])

diff --git a/Engine/Source/Runtime/DeveloperSettings/Private/Engine/PlatformSettingsManager.cpp b/Engine/Source/Runtime/DeveloperSettings/Private/Engine/PlatformSettingsManager.cpp
--- a/Engine/Source/Runtime/DeveloperSettings/Private/Engine/PlatformSettingsManager.cpp
+++ b/Engine/Source/Runtime/DeveloperSettings/Private/Engine/PlatformSettingsManager.cpp
@@ -43,7 +43,17 @@ UPlatformSettings* UPlatformSettingsManager::GetSettingsForPlatform(TSubclassOf<
 	if (Result == nullptr)
 	{
 		Result = CreateSettingsObjectForPlatform(SettingsClass, DesiredPlatformIniName);
+
+		if (DesiredPlatformIniName == IniPlatformName)
+		{
+			Record.PlatformInstance = Result;
+		}
+		else
+		{
+			Record.OtherPlatforms.Add(DesiredPlatformIniName, Result);
+		}
 	}
+
 	return Result;
 }
 
@@ -62,7 +72,7 @@ UPlatformSettings* UPlatformSettingsManager::CreateSettingsObjectForPlatform(TSu
 
 		PlatformSettingsForThisClass->LoadConfig();
 	}
-	
+
 	return PlatformSettingsForThisClass;
 }
 
