Fixed a copy from TArray<UPlatformSettings*> into TArray<TObjectPtr<UPlatformSettings>> which should be a move.

diff --git a/Engine/Source/Runtime/DeveloperSettings/Private/Engine/PlatformSettingsManager.cpp b/Engine/Source/Runtime/DeveloperSettings/Private/Engine/PlatformSettingsManager.cpp
--- a/Engine/Source/Runtime/DeveloperSettings/Private/Engine/PlatformSettingsManager.cpp
+++ b/Engine/Source/Runtime/DeveloperSettings/Private/Engine/PlatformSettingsManager.cpp
@@ -102,9 +102,9 @@ TArray<FName> UPlatformSettingsManager::GetKnownAndEnablePlatformIniNames()
 	return Results;
 }
 
-TArray<UPlatformSettings*> UPlatformSettingsManager::GetAllPlatformSettings(TSubclassOf<UPlatformSettings> SettingsClass) const
+TArray<TObjectPtr<UPlatformSettings>> UPlatformSettingsManager::GetAllPlatformSettings(TSubclassOf<UPlatformSettings> SettingsClass) const
 {
-	TArray<UPlatformSettings*> Settings;
+	TArray<TObjectPtr<UPlatformSettings>> Settings;
 	for (FName PlatformIniName : GetKnownAndEnablePlatformIniNames())
 	{
 		Settings.Add(GetSettingsForPlatform(SettingsClass, PlatformIniName));
diff --git a/Engine/Source/Runtime/DeveloperSettings/Public/Engine/PlatformSettingsManager.h b/Engine/Source/Runtime/DeveloperSettings/Public/Engine/PlatformSettingsManager.h
--- a/Engine/Source/Runtime/DeveloperSettings/Public/Engine/PlatformSettingsManager.h
+++ b/Engine/Source/Runtime/DeveloperSettings/Public/Engine/PlatformSettingsManager.h
@@ -75,12 +75,12 @@ public:
 	static DEVELOPERSETTINGS_API TArray<FName> GetKnownAndEnablePlatformIniNames();
 
 	template <typename TPlatformSettingsClass>
-	FORCEINLINE TArray<UPlatformSettings*> GetAllPlatformSettings() const
+	FORCEINLINE TArray<TObjectPtr<UPlatformSettings>> GetAllPlatformSettings() const
 	{
 		return GetAllPlatformSettings(TPlatformSettingsClass::StaticClass());
 	}
 
-	DEVELOPERSETTINGS_API TArray<UPlatformSettings*> GetAllPlatformSettings(TSubclassOf<UPlatformSettings> SettingsClass) const;
+	DEVELOPERSETTINGS_API TArray<TObjectPtr<UPlatformSettings>> GetAllPlatformSettings(TSubclassOf<UPlatformSettings> SettingsClass) const;
 
 	template <typename TPlatformSettingsClass>
 	FORCEINLINE TPlatformSettingsClass* GetSettingsForPlatform(FName TargetIniPlatformName) const
diff --git a/Engine/Source/Runtime/Engine/Private/UserInterface/InputSettings.cpp b/Engine/Source/Runtime/Engine/Private/UserInterface/InputSettings.cpp
--- a/Engine/Source/Runtime/Engine/Private/UserInterface/InputSettings.cpp
+++ b/Engine/Source/Runtime/Engine/Private/UserInterface/InputSettings.cpp
@@ -681,9 +681,9 @@ const TArray<FName>& UInputPlatformSettings::GetAllHardwareDeviceNames()
 	HardwareDevices.Add(FHardwareDeviceIdentifier::DefaultKeyboardAndMouse.HardwareDeviceIdentifier);
 
 	// Get every known platform's InputPlatformSettings and compile a list of them
-	TArray<UPlatformSettings*> AllInputSettings = UPlatformSettingsManager::Get().GetAllPlatformSettings<UInputPlatformSettings>();
+	TArray<TObjectPtr<UPlatformSettings>> AllInputSettings = UPlatformSettingsManager::Get().GetAllPlatformSettings<UInputPlatformSettings>();
 
-	for (const UPlatformSettings* Setting : AllInputSettings)
+	for (const TObjectPtr<UPlatformSettings>& Setting : AllInputSettings)
 	{
 		if (const UInputPlatformSettings* InputSetting = Cast<UInputPlatformSettings>(Setting))
 		{
