Added a Get Current Platform Name function to Platform Setting Manager that will return the Current Platform or the Simulated Platform's Name.

#tests PIE
#rnx

diff --git a/Engine/Source/Runtime/DeveloperSettings/Private/Engine/PlatformSettingsManager.cpp b/Engine/Source/Runtime/DeveloperSettings/Private/Engine/PlatformSettingsManager.cpp
--- a/Engine/Source/Runtime/DeveloperSettings/Private/Engine/PlatformSettingsManager.cpp
+++ b/Engine/Source/Runtime/DeveloperSettings/Private/Engine/PlatformSettingsManager.cpp
@@ -50,6 +50,17 @@ UPlatformSettings* UPlatformSettingsManager::GetSettingsForPlatform(TSubclassOf<
 	return Result;
 }
 
+FName UPlatformSettingsManager::GetCurrentPlatformName() const
+{
+#if WITH_EDITOR
+	if (GIsEditor && !SimulatedEditorPlatform.IsNone())
+	{
+		return SimulatedEditorPlatform;
+	}
+#endif
+	return IniPlatformName;
+}
+
 UPlatformSettings* UPlatformSettingsManager::CreateSettingsObjectForPlatform(TSubclassOf<UPlatformSettings> SettingsClass, FName TargetIniPlatformName) const
 {
 	const FString PlatformSettingsName = SettingsClass->GetName() + TEXT("_") + TargetIniPlatformName.ToString();
diff --git a/Engine/Source/Runtime/DeveloperSettings/Public/Engine/PlatformSettingsManager.h b/Engine/Source/Runtime/DeveloperSettings/Public/Engine/PlatformSettingsManager.h
--- a/Engine/Source/Runtime/DeveloperSettings/Public/Engine/PlatformSettingsManager.h
+++ b/Engine/Source/Runtime/DeveloperSettings/Public/Engine/PlatformSettingsManager.h
@@ -55,6 +55,8 @@ public:
 	// Returns the platform settings for the specified class for the specified platform
 	DEVELOPERSETTINGS_API UPlatformSettings* GetSettingsForPlatform(TSubclassOf<UPlatformSettings> SettingsClass, FName DesiredPlatformIniName) const;
 
+	DEVELOPERSETTINGS_API FName GetCurrentPlatformName() const;
+
 #if WITH_EDITOR
 	// Gets the simulated platform
 	static FName GetEditorSimulatedPlatform() { return SimulatedEditorPlatform; }
