Fix synchronous load done at a bad time because UCommonInputPlatformSettings was expecting its PostLoad method to be called to force the load, but it was never happening.
Now PostLoad will be called for all UPlatformSettings.

#tests Validated with Lyra in editor and cook mode that common input platform settings where still loaded correctly, while PostLoad is now called as expected, resulting in sync load no longer being called at later stage. Tested in cooked game, no regression.

diff --git a/Engine/Source/Runtime/DeveloperSettings/Private/Engine/PlatformSettingsManager.cpp b/Engine/Source/Runtime/DeveloperSettings/Private/Engine/PlatformSettingsManager.cpp
--- a/Engine/Source/Runtime/DeveloperSettings/Private/Engine/PlatformSettingsManager.cpp
+++ b/Engine/Source/Runtime/DeveloperSettings/Private/Engine/PlatformSettingsManager.cpp
@@ -74,7 +74,8 @@ UPlatformSettings* UPlatformSettingsManager::CreateSettingsObjectForPlatform(TSu
 	{
 		PlatformSettingsForThisClass->InitializePlatformDefaults();
 
-		PlatformSettingsForThisClass->LoadConfig();
+		//With this call, it both load the config and call post load on the object, which would not be done if directly calling LoadConfig on the object.
+		PlatformSettingsForThisClass->PostLoad();
 	}
 	
 	return PlatformSettingsForThisClass;
