Fix common input settings regression (no longer editable in the project settings):
- Platfom inialization is necessary also for the CDO, reactivated it.
- Reenable defered loading on editor only, and keep sync loading in CreateSettingsObjectForPlatform otherwise.

#rnx

diff --git a/Engine/Plugins/Runtime/CommonUI/Source/CommonInput/Private/CommonInputSettings.cpp b/Engine/Plugins/Runtime/CommonUI/Source/CommonInput/Private/CommonInputSettings.cpp
--- a/Engine/Plugins/Runtime/CommonUI/Source/CommonInput/Private/CommonInputSettings.cpp
+++ b/Engine/Plugins/Runtime/CommonUI/Source/CommonInput/Private/CommonInputSettings.cpp
@@ -16,10 +16,7 @@ UCommonInputSettings::UCommonInputSettings(const FObjectInitializer& Initializer
 	: Super(Initializer)
 	, bInputDataLoaded(false)
 {
-	if (!HasAnyFlags(RF_ClassDefaultObject))
-	{
-		PlatformInput.Initialize(UCommonInputPlatformSettings::StaticClass());
-	}
+	PlatformInput.Initialize(UCommonInputPlatformSettings::StaticClass());
 }
 
 void UCommonInputSettings::LoadData()
diff --git a/Engine/Source/Runtime/DeveloperSettings/Private/Engine/PlatformSettingsManager.cpp b/Engine/Source/Runtime/DeveloperSettings/Private/Engine/PlatformSettingsManager.cpp
--- a/Engine/Source/Runtime/DeveloperSettings/Private/Engine/PlatformSettingsManager.cpp
+++ b/Engine/Source/Runtime/DeveloperSettings/Private/Engine/PlatformSettingsManager.cpp
@@ -74,8 +74,12 @@ UPlatformSettings* UPlatformSettingsManager::CreateSettingsObjectForPlatform(TSu
 	{
 		PlatformSettingsForThisClass->InitializePlatformDefaults();
 
+#if WITH_EDITOR
+		PlatformSettingsForThisClass->LoadConfig();
+#else
 		//Note: Calling PostLoad will implicitly call LoadConfig. Calling just LoadConfig does not call PostLoad.
 		PlatformSettingsForThisClass->PostLoad();
+#endif		
 	}
 	
 	return PlatformSettingsForThisClass;
