In DeveloperSettings::ImportConsoleVariableValues downgrade error handling when a console variable was not found from fatal to error
Update the Log Category from LogTemp to LogInit

#Synced-CL 52499590

diff --git a/Engine/Source/Runtime/DeveloperSettings/Private/Engine/DeveloperSettings.cpp b/Engine/Source/Runtime/DeveloperSettings/Private/Engine/DeveloperSettings.cpp
--- a/Engine/Source/Runtime/DeveloperSettings/Private/Engine/DeveloperSettings.cpp
+++ b/Engine/Source/Runtime/DeveloperSettings/Private/Engine/DeveloperSettings.cpp
@@ -124,12 +124,12 @@ void UDeveloperSettings::ImportConsoleVariableValues()
 			{
 				if (Property->ImportText_InContainer(*CVar->GetString(), this, this, PPF_ConsoleVariable) == NULL)
 				{
-					UE_LOG(LogTemp, Error, TEXT("%s import failed for %s on console variable %s (=%s)"), *GetClass()->GetName(), *Property->GetName(), *CVarName, *CVar->GetString());
+					UE_LOG(LogInit, Error, TEXT("%s import failed for %s on console variable %s (=%s)"), *GetClass()->GetName(), *Property->GetName(), *CVarName, *CVar->GetString());
 				}
 			}
 			else
 			{
-				UE_LOG(LogTemp, Fatal, TEXT("%s failed to find console variable %s for %s"), *GetClass()->GetName(), *CVarName, *Property->GetName());
+				UE_LOG(LogInit, Error, TEXT("%s failed to find console variable %s for %s"), *GetClass()->GetName(), *CVarName, *Property->GetName());
 			}
 		}
 	}
